A subsystem to bubble up notifications from the StallDetector
* The subsystem adds a new toolbar widget in the level editor which changes to a notified state when a stall is detected.
* Clicking the widget brings up a tab with a table of the detected stall history.  Each stall has some information that can be copied to a clipboard for users to dump the data into a issue ticket or chat message.
* NOTE: Currently iconography is all placeholders. Need UX pass to improve.

diff --git a/Engine/Plugins/Editor/EditorPerformanceModule/EditorPerformanceModule.uplugin b/Engine/Plugins/Editor/EditorPerformanceModule/EditorPerformanceModule.uplugin
new file mode 100644
--- /dev/null
+++ b/Engine/Plugins/Editor/EditorPerformanceModule/EditorPerformanceModule.uplugin
@@ -0,0 +1,25 @@
+﻿{
+	"FileVersion": 3,
+	"Version": 1,
+	"VersionName": "1.0",
+	"FriendlyName": "Editor Performance Module",
+	"Description": "Provides Editor Performance Measurement Tools",
+	"Category": "Editor",
+	"CreatedBy": "Epic Games, Inc.",
+	"CreatedByURL": "http://epicgames.com",
+	"DocsURL": "",
+	"MarketplaceURL": "",
+	"SupportURL": "",
+	"EnabledByDefault": false,
+	"CanContainContent": false,
+	"IsBetaVersion": false,
+	"IsExperimentalVersion": true,
+	"Installed": false,
+	"Modules": [
+		{
+			"Name": "EditorPerformanceModule",
+			"Type": "Editor",
+			"LoadingPhase": "PostConfigInit"
+		}
+	]
+}
\ No newline at end of file
diff --git a/Engine/Plugins/Editor/EditorPerformanceModule/Source/EditorPerformanceModule/EditorPerformanceModule.Build.cs b/Engine/Plugins/Editor/EditorPerformanceModule/Source/EditorPerformanceModule/EditorPerformanceModule.Build.cs
new file mode 100644
--- /dev/null
+++ b/Engine/Plugins/Editor/EditorPerformanceModule/Source/EditorPerformanceModule/EditorPerformanceModule.Build.cs
@@ -0,0 +1,34 @@
+﻿// Copyright Epic Games, Inc. All Rights Reserved.
+
+using UnrealBuildTool;
+
+public class EditorPerformanceModule : ModuleRules
+{
+	public EditorPerformanceModule(ReadOnlyTargetRules Target) : base(Target)
+	{
+		PCHUsage = ModuleRules.PCHUsageMode.UseExplicitOrSharedPCHs;
+		
+		PublicDependencyModuleNames.AddRange(
+			new string[]
+			{
+				"Core",
+			}
+		);
+		
+		PrivateDependencyModuleNames.AddRange(
+			new string[]
+			{
+				"ApplicationCore",
+				"CoreUObject",
+				"Engine",
+				"EditorSubsystem",
+				"InputCore",
+				"MessageLog",
+				"SlateCore",
+				"Slate",
+				"ToolMenus",
+				"WorkspaceMenuStructure"
+			}
+		);
+	}
+}
\ No newline at end of file
diff --git a/Engine/Plugins/Editor/EditorPerformanceModule/Source/EditorPerformanceModule/Private/EditorPerformanceModule.cpp b/Engine/Plugins/Editor/EditorPerformanceModule/Source/EditorPerformanceModule/Private/EditorPerformanceModule.cpp
new file mode 100644
--- /dev/null
+++ b/Engine/Plugins/Editor/EditorPerformanceModule/Source/EditorPerformanceModule/Private/EditorPerformanceModule.cpp
@@ -0,0 +1,9 @@
+﻿// Copyright Epic Games, Inc. All Rights Reserved.
+
+#include "EditorPerformanceModule.h"
+
+class FEditorPerformanceModule : public IEditorPerformanceModule
+{
+};
+
+IMPLEMENT_MODULE(FEditorPerformanceModule, EditorPerformanceModule)
diff --git a/Engine/Plugins/Editor/EditorPerformanceModule/Source/EditorPerformanceModule/Private/StallLogSubsystem/StallLogSubsystem.cpp b/Engine/Plugins/Editor/EditorPerformanceModule/Source/EditorPerformanceModule/Private/StallLogSubsystem/StallLogSubsystem.cpp
new file mode 100644
--- /dev/null
+++ b/Engine/Plugins/Editor/EditorPerformanceModule/Source/EditorPerformanceModule/Private/StallLogSubsystem/StallLogSubsystem.cpp
@@ -0,0 +1,732 @@
+// Copyright Epic Games, Inc. All Rights Reserved.
+
+#include "StallLogSubsystem/StallLogSubsystem.h"
+
+#include "Async/TaskGraphInterfaces.h"
+#include "Containers/Array.h"
+#include "Framework/Notifications/NotificationManager.h"
+#include "GenericPlatform/GenericPlatformCrashContext.h"
+#include "HAL/IConsoleManager.h"
+#include "HAL/PlatformApplicationMisc.h"
+#include "Logging/MessageLog.h"
+#include "MessageLogInitializationOptions.h"
+#include "MessageLogModule.h"
+#include "Modules/ModuleManager.h"
+#include "ProfilingDebugging/StallDetector.h"
+#include "SlateOptMacros.h"
+#include "Stats/Stats2.h"
+#include "ToolMenus.h"
+#include "Widgets/Docking/SDockTab.h"
+#include "Widgets/Layout/SHeader.h"
+#include "Widgets/Layout/SSpacer.h"
+#include "Widgets/Notifications/SNotificationList.h"
+#include "Widgets/Views/STableRow.h"
+#include "WorkspaceMenuStructure.h"
+#include "WorkspaceMenuStructureModule.h"
+
+#define LOCTEXT_NAMESPACE "StallLogSubsystem"
+
+static bool GEnableStallLogSubsystem = true;
+static FAutoConsoleVariableRef CCvarEnableStallLogging(
+	TEXT("Editor.StallLogger.Enable"),
+	GEnableStallLogSubsystem,
+	TEXT("Whether the editor stall logger subsystem is enabled."),
+	ECVF_Default);
+
+/**
+ * Metadata for each detected stall of the application
+ */
+class FStallLogItem : public TSharedFromThis<FStallLogItem>
+{
+public:
+	FStallLogItem() = default;
+	FStallLogItem(FString InLocation, float InDurationSeconds, FDateTime InTime, TArray<FString> InStackTrace);
+
+	FString Location;
+	float DurationSeconds;
+	FDateTime Time;
+	TArray<FString> StackTrace;
+};
+
+FStallLogItem::FStallLogItem(FString InLocation, float InDurationSeconds, FDateTime InTime, TArray<FString> InStackTrace)
+	: Location(MoveTemp(InLocation))
+	, DurationSeconds(InDurationSeconds)
+	, Time(InTime)
+	, StackTrace(MoveTemp(InStackTrace))
+{}
+
+using FStallLogItemPtr = TSharedPtr<FStallLogItem>;
+
+/**
+ * Holds a history of all the detected stalls
+ * Model used for the UI
+ */
+class FStallLogHistory
+{
+public:
+	void OnStallDetected(uint64 UniqueID, FDateTime InDetectTime, FStringView StatName, TArray<FString> StackTrace, uint32 ThreadID);
+	void OnStallCompleted(uint64 UniqueID, double InDurationSeconds);
+	
+	void ClearStallLog();
+	
+	const TArray<FStallLogItemPtr>& GetStallLog() const;
+
+private:
+	struct FInFlightStall
+	{
+		FDateTime DetectTime;
+		FString StatName;
+		TArray<FString> StackTrace;
+		uint32 ThreadID;
+	};
+
+	TMap<uint64, FInFlightStall> InFlightStalls;
+	TArray<FStallLogItemPtr> StallLogs;
+};
+
+const TArray<FStallLogItemPtr>& FStallLogHistory::GetStallLog() const
+{
+	checkf(IsInGameThread(), TEXT("Can only be run on GameThread"));
+	return StallLogs;
+}
+
+void FStallLogHistory::ClearStallLog()
+{
+	checkf(IsInGameThread(), TEXT("Can only be run on GameThread"));
+	StallLogs.Empty();
+}
+
+namespace
+{
+	const FName StallLogTabName = FName(TEXT("StallLogTab"));
+
+	const FName ColumnName_Location = FName(TEXT("Location"));
+	const FName ColumnName_Duration = FName(TEXT("Duration"));
+	const FName ColumnName_Time = FName(TEXT("Time"));
+	const FName ColumnName_Copy = FName(TEXT("Copy"));
+
+	DECLARE_DELEGATE(FStallLogClearLog)
+	DECLARE_DELEGATE_RetVal(const FSlateBrush*, FGetSlateBrush);
+	
+	/**
+	 * A widget for each row of the stall stable
+	 */
+	class SStallLogItemRow
+		: public SMultiColumnTableRow<FStallLogItemPtr>
+	{
+	public:
+		SLATE_BEGIN_ARGS(SStallLogItemRow) {}
+		SLATE_END_ARGS()
+
+	public:
+
+		void Construct(const FArguments& InArgs, const TSharedRef<STableViewBase>& InOwnerTableView, const FStallLogItemPtr& InStallLogItem)
+		{
+			StallLogItem = InStallLogItem;
+			SMultiColumnTableRow<FStallLogItemPtr>::Construct(FSuperRowType::FArguments(), InOwnerTableView);
+		}
+		
+		virtual TSharedRef<SWidget> GenerateWidgetForColumn(const FName& ColumnName) override
+		{
+			if (ColumnName == ColumnName_Location)
+			{
+				return SNew(SBox)
+					.Padding(FMargin(4.0f, 0.0f))
+					.VAlign(VAlign_Center)
+					[
+						SNew(STextBlock)
+							.ColorAndOpacity(FSlateColor::UseForeground())
+							.Text(FText::FromString(StallLogItem->Location))
+					];
+			}
+			else if (ColumnName == ColumnName_Duration)
+			{
+				return SNew(SBox)
+					.Padding(FMargin(4.0f, 0.0f))
+					.VAlign(VAlign_Center)
+					[
+						SNew(STextBlock)
+							.ColorAndOpacity(FSlateColor::UseForeground())
+							.Text(FText::Format(LOCTEXT("DurationFmt", "{0}"), StallLogItem->DurationSeconds))
+					];
+			}
+			else if (ColumnName == ColumnName_Time)
+			{
+				return SNew(SBox)
+					.Padding(FMargin(4.0f, 0.0f))
+					.VAlign(VAlign_Center)
+					[
+						SNew(STextBlock)
+							.ColorAndOpacity(FSlateColor::UseForeground())
+							.Text(FText::Format(LOCTEXT("TimeFmt", "{0}"), FText::AsDateTime(StallLogItem->Time)))
+					];
+			}
+			else if (ColumnName == ColumnName_Copy)
+			{
+				return
+				SNew(SHorizontalBox)
+				+SHorizontalBox::Slot()
+				[
+					SNew(SSpacer)
+				]
+				+SHorizontalBox::Slot()
+				.MaxWidth(16)
+				[
+					SNew(SButton)
+					.ToolTipText(LOCTEXT("StallDetector", "Copy Stall Information"))
+					.ButtonStyle(FAppStyle::Get(), "NoBorder")
+					.ContentPadding(0)
+					.Visibility(EVisibility::Visible)
+					.OnClicked_Lambda([this]()
+					{
+						// Build a string of the stall information and put it on the clipboard
+						FPlatformApplicationMisc::ClipboardCopy(*FString::Join(StallLogItem->StackTrace, TEXT("\n")));
+
+						FNotificationInfo Info(LOCTEXT("StallLogInfoCopied", "Copied to clipboard"));
+						Info.ExpireDuration = 2.0f;
+						FSlateNotificationManager::Get().AddNotification(Info);
+						
+						return FReply::Handled();
+					})
+					[
+						SNew(SImage)
+						.Image(FAppStyle::GetBrush("GenericCommands.Copy"))
+					]
+				]
+				+SHorizontalBox::Slot()
+				[
+					SNew(SSpacer)
+				];
+			}
+
+			return SNullWidget::NullWidget;
+		}
+
+	private:
+		FStallLogItemPtr StallLogItem;
+	};
+
+	/**
+	 * A widget to display the table of stalls
+	 */
+	class SStallLog
+		: public SCompoundWidget
+	{
+	public:
+		SLATE_BEGIN_ARGS( SStallLog ) {}
+			SLATE_EVENT(FStallLogClearLog, OnClearLog)
+			SLATE_ARGUMENT (const TArray<TSharedPtr<FStallLogItem>>*, StallLogItems)
+		SLATE_END_ARGS()
+
+		virtual ~SStallLog();
+
+		void Construct(const FArguments& InArgs);
+
+	private:
+		
+		TSharedPtr<SListView<FStallLogItemPtr>> ListViewPtr;
+		FStallLogClearLog ClearLogDelegate;
+	};
+
+	SStallLog::~SStallLog()	= default;
+	
+	void SStallLog::Construct(const FArguments& InArgs)
+	{
+		auto StallLogGenerateRow = [](FStallLogItemPtr StallLogItem, const TSharedRef<STableViewBase>& OwnerTable) -> TSharedRef<ITableRow>
+		{
+			return SNew(SStallLogItemRow, OwnerTable, StallLogItem);
+		};
+
+		ClearLogDelegate = InArgs._OnClearLog;
+		
+		ChildSlot
+		.Padding(3)
+		[
+			SNew(SVerticalBox)
+
+			// Table
+			+SVerticalBox::Slot()
+			.AutoHeight()
+			.VAlign(EVerticalAlignment::VAlign_Top)
+			.Padding(FMargin(0.0f, 4.0f, 0.0f, 4.0f))
+			[
+				SNew(SSplitter)
+				.PhysicalSplitterHandleSize(2.0f)
+				.ResizeMode(ESplitterResizeMode::FixedSize)
+				+ SSplitter::Slot()
+				.Value(0.15f)
+				[
+					SNew(SBox)
+					.Padding(FMargin(4.f))
+					[
+						SNew(SBorder)
+						.Padding(FMargin(0))
+						.BorderImage(FAppStyle::GetBrush("Brushes.Recessed"))
+						[
+							SAssignNew(ListViewPtr, SListView<TSharedPtr<FStallLogItem>>)
+								.ListItemsSource(InArgs._StallLogItems)
+								.OnGenerateRow_Lambda(StallLogGenerateRow)
+								.SelectionMode(ESelectionMode::None)
+								.HeaderRow
+								(
+									SNew(SHeaderRow)
+
+									+ SHeaderRow::Column(ColumnName_Location)
+										.DefaultLabel(LOCTEXT("StallLogColumnHeader_StallDetectorName", "Stall Detector Name"))
+										.FillWidth(0.3f)
+
+									+ SHeaderRow::Column(ColumnName_Duration)
+										.DefaultLabel(LOCTEXT("StallLogColumnHeader_Duration", "Duration"))
+										.FillWidth(0.3f)
+
+									+ SHeaderRow::Column(ColumnName_Time)
+										.DefaultLabel(LOCTEXT("StallLogColumnHeader_Time", "Time Of Stall"))
+										.FillWidth(0.3f)
+
+									+ SHeaderRow::Column(ColumnName_Copy)
+										.DefaultLabel(LOCTEXT("StallLogColumnHeader_CopyButton", "Copy Stall Info"))
+										.FillWidth(0.1f)
+								)
+								
+						]
+					]
+				]
+			]
+
+			+SVerticalBox::Slot()
+			[
+				SNew(SSpacer)
+			]
+
+			+SVerticalBox::Slot()
+			.AutoHeight()
+			.VAlign(EVerticalAlignment::VAlign_Bottom)
+			[
+				SNew(SButton)
+				.OnClicked_Lambda([this]()
+				{
+					const bool Executed = ClearLogDelegate.ExecuteIfBound();
+					return Executed ? FReply::Handled() : FReply::Unhandled();
+				})
+				[
+				SNew(SHorizontalBox)
+				+ SHorizontalBox::Slot()
+				.AutoWidth()
+				.VAlign(VAlign_Center)
+				[
+					SNew(STextBlock)
+					.Text(LOCTEXT("StallLog_Clear", "Clear Stall Log"))
+					.Justification(ETextJustify::Center)
+				]
+				]
+			]
+		];
+	}
+	
+	class SStallLogStatusBarWidget : public SCompoundWidget
+	{
+		SLATE_BEGIN_ARGS(SStallLogStatusBarWidget) {}
+			SLATE_ARGUMENT(FGetSlateBrush, GetFilterBadgeIcon)
+		SLATE_END_ARGS()
+
+		void Construct(const FArguments& InArgs);
+
+	private:
+		
+		FText						GetToolTipText() const;
+		const FSlateBrush*			GetBadgeIcon() const;
+
+		FGetSlateBrush GetFilterBadgeIconDelegate;
+	};
+
+	FText SStallLogStatusBarWidget::GetToolTipText() const
+	{
+		return LOCTEXT("StallLogStatusBarToolTip", "Opens the Stall Log");
+	}
+
+	const FSlateBrush* SStallLogStatusBarWidget::GetBadgeIcon() const
+	{
+		if (GetFilterBadgeIconDelegate.IsBound())
+		{
+			return GetFilterBadgeIconDelegate.Execute();
+		}
+		return nullptr;
+	}
+
+	void SStallLogStatusBarWidget::Construct(const FArguments& InArgs)
+	{
+		GetFilterBadgeIconDelegate = InArgs._GetFilterBadgeIcon;
+		this->ChildSlot
+		[
+			SNew(SButton)
+			.ButtonStyle(&FAppStyle::Get().GetWidgetStyle<FButtonStyle>("StatusBar.StatusBarButton"))
+			.ContentPadding(FMargin(6.0f, 0.0f))
+			[
+				SNew(SHorizontalBox)
+				+ SHorizontalBox::Slot()
+				.AutoWidth()
+				.VAlign(VAlign_Center)
+				.Padding(0, 0, 3, 0)
+				[
+					SNew(SOverlay)
+
+					+ SOverlay::Slot()
+					.HAlign(HAlign_Center)
+					.VAlign(VAlign_Top)
+					[
+						SNew(SImage)
+						.ColorAndOpacity(FSlateColor::UseForeground())
+						.Image_Lambda([this] { return GetBadgeIcon();  })
+						.ToolTipText_Lambda([this] { return GetToolTipText(); })
+					]
+				]
+			]
+			.OnClicked_Lambda([]()
+			{
+				FGlobalTabmanager::Get()->TryInvokeTab(FTabId(StallLogTabName));
+				return FReply::Handled();
+			})
+		];
+	}
+
+}
+
+void FStallLogHistory::OnStallDetected(
+	uint64 UniqueID,
+	FDateTime InDetectTime,
+	FStringView StatName,
+	TArray<FString> StackTrace,
+	uint32 ThreadID)
+{
+	checkf(IsInGameThread(), TEXT("Can only be run on GameThread"));
+	FInFlightStall InFlightStall;
+	InFlightStall.DetectTime = InDetectTime;
+	InFlightStall.StatName = StatName;
+	InFlightStall.StackTrace = MoveTemp(StackTrace);
+	InFlightStall.ThreadID = ThreadID;
+
+	InFlightStalls.Emplace(UniqueID, MoveTemp(InFlightStall));
+}
+
+void FStallLogHistory::OnStallCompleted(uint64 UniqueID, double InDurationSeconds)
+{
+	checkf(IsInGameThread(), TEXT("Can only be run on GameThread"));
+	
+	// Find the completed stall and only add it to the history if found
+	FInFlightStall* InFlightStall = InFlightStalls.Find(UniqueID);
+	ensure(InFlightStall != nullptr);
+	
+	if (InFlightStall != nullptr)
+	{
+		TSharedPtr<FStallLogItem> StallLogItem = MakeShared<FStallLogItem>(
+			MoveTemp(InFlightStall->StatName), 
+			InDurationSeconds,
+			InFlightStall->DetectTime,
+			MoveTemp(InFlightStall->StackTrace));
+		StallLogs.Emplace(StallLogItem);
+
+		InFlightStalls.Remove(UniqueID);
+	}
+}
+
+DECLARE_CYCLE_STAT(
+	TEXT("StallLoggerSubsystem"),
+	STAT_FDelegateGraphTask_StallLogger,
+	STATGROUP_TaskGraphTasks);
+
+namespace
+{	
+	void RegisterStallsLogListing()
+	{
+		FMessageLogInitializationOptions InitOptions;
+		InitOptions.bShowFilters = true;
+		InitOptions.bAllowClear = true;
+
+		FMessageLogModule& MessageLogModule = FModuleManager::LoadModuleChecked<FMessageLogModule>("MessageLog");
+		MessageLogModule.RegisterLogListing("StallLog", LOCTEXT("StallLog", "Editor Stall Logger"), InitOptions);
+	}
+
+	void UnregisterStallsLogListing()
+	{
+		FMessageLogModule& MessageLogModule = FModuleManager::LoadModuleChecked<FMessageLogModule>("MessageLog");
+		MessageLogModule.UnregisterLogListing("StallLog");
+	}
+}
+
+bool UStallLogSubsystem::ShouldCreateSubsystem(UObject* Outer) const
+{
+#if !WITH_EDITOR
+	return false;
+#endif
+	
+	if (!FSlateApplication::IsInitialized())
+	{
+		return false;
+	}
+	
+	return Super::ShouldCreateSubsystem(Outer);
+}
+
+void UStallLogSubsystem::Initialize(FSubsystemCollectionBase& Collection)
+{
+	UEditorSubsystem::Initialize(Collection);
+
+	StallLogHistory = MakeShared<FStallLogHistory>();
+
+	RegisterStallsLogListing();
+	RegisterStallDetectedDelegates();
+
+	/** Register a tab spawner invoked by the tools bar button */
+	{
+		const FSlateIcon StallLogIcon(FAppStyle::GetAppStyleSetName(), "EditorViewport.ToggleRealTime");
+
+		FGlobalTabmanager::Get()->RegisterNomadTabSpawner(StallLogTabName, FOnSpawnTab::CreateUObject(this, &UStallLogSubsystem::CreateStallLogTab))
+			.SetDisplayName(LOCTEXT("StallLogTabTitle", "Stall Log"))
+			.SetTooltipText(LOCTEXT("StallLogTabToolTipText", "Show Stall Log"))
+			.SetGroup(WorkspaceMenu::GetMenuStructure().GetToolsCategory())
+			.SetIcon(StallLogIcon);
+	}
+	
+	/** Add the widget button to the tools bar */	
+	UToolMenu* Menu = UToolMenus::Get()->ExtendMenu("LevelEditor.StatusBar.ToolBar");
+
+	FToolMenuSection& StallDetectorSection = Menu->AddSection(
+		"StallLog", FText::GetEmpty(), FToolMenuInsert("DDC", EToolMenuInsertType::Before));
+
+	auto CreateStallLogWidget = [this]()
+	{
+		return SNew(SStallLogStatusBarWidget)
+			.GetFilterBadgeIcon(FGetSlateBrush::CreateUObject(this, &UStallLogSubsystem::GetStatusBarBadgeIcon));
+	};
+	
+	StallDetectorSection.AddEntry(
+		FToolMenuEntry::InitWidget("StallLogStatusBar", CreateStallLogWidget(), FText::GetEmpty(), true, false));
+}
+
+void UStallLogSubsystem::Deinitialize()
+{
+	FGlobalTabmanager::Get()->UnregisterNomadTabSpawner(StallLogTabName);
+
+	if (const TSharedPtr<SDockTab> StallLogTabShared = StallLogTab.Pin())
+	{
+		StallLogTabShared->RequestCloseTab();
+	}
+
+	UnregisterStallDetectedDelegates();
+	UnregisterStallsLogListing();
+}
+	
+TSharedRef<SDockTab> UStallLogSubsystem::CreateStallLogTab(const FSpawnTabArgs& InArgs)
+{
+	return SAssignNew(StallLogTab, SDockTab)
+	.TabRole(ETabRole::NomadTab)
+	[
+		SAssignNew(StallLog, SStallLog)
+			.StallLogItems(&StallLogHistory->GetStallLog())
+			.OnClearLog_Lambda([StallLogHistoryWkPtr = this->StallLogHistory.ToWeakPtr()]
+			{
+				TSharedPtr<FStallLogHistory> StallLogHistoryPtr = StallLogHistoryWkPtr.Pin();
+				if (StallLogHistoryPtr)
+				{
+					StallLogHistoryPtr->ClearStallLog();
+				}
+			})
+	];
+}
+
+const FSlateBrush* UStallLogSubsystem::GetStatusBarBadgeIcon() const
+{
+	if (!StallLogHistory->GetStallLog().IsEmpty())
+	{
+		return FAppStyle::GetBrush(TEXT("EditorViewport.ToggleRealTime"));
+	}
+	else
+	{
+		return FAppStyle::GetBrush(TEXT("Level.SaveDisabledIcon16x"));
+	}
+}
+
+void UStallLogSubsystem::RegisterStallDetectedDelegates()
+{
+	OnStallDetectedDelegate = UE::FStallDetector::StallDetected.AddLambda(
+			[StallLogHistory = this->StallLogHistory](const UE::FStallDetectedParams& Params)
+			{
+				if (!GEnableStallLogSubsystem)
+				{
+					return;
+				}
+
+				FDateTime Now = FDateTime::Now();
+
+				// Add a bookmark to Insights when a stall is detected
+				// Helps understand the context of the stall on the thread timeline
+				TRACE_BOOKMARK(TEXT("Stall [%s]"), *FString(Params.StatName));
+				
+				// Grab a stacktrace of the stalled thread
+				const uint32 MaxStackDepth = 64;
+				TArray<uint64> Backtrace;
+				Backtrace.SetNumUninitialized(MaxStackDepth);
+				const uint32 StackDepth = FPlatformStackWalk::CaptureThreadStackBackTrace(
+					Params.ThreadId, Backtrace.GetData(), MaxStackDepth);
+
+				Backtrace.SetNum(StackDepth);
+
+				TArray<FString> StackTrace;
+				StackTrace.SetNum(StackDepth);
+
+				for (uint32 StackFrameIndex = 0; StackFrameIndex < StackDepth; ++StackFrameIndex)
+				{
+					FProgramCounterSymbolInfo SymbolInfo;
+					FPlatformStackWalk::ProgramCounterToSymbolInfo(Backtrace[StackFrameIndex], SymbolInfo);
+
+					// Strip module path.
+					const ANSICHAR* Pos0 = FCStringAnsi::Strrchr( SymbolInfo.ModuleName, '\\' );
+					const ANSICHAR* Pos1 = FCStringAnsi::Strrchr( SymbolInfo.ModuleName, '/' );
+					const UPTRINT RealPos = FMath::Max(reinterpret_cast<UPTRINT>(Pos0), reinterpret_cast<UPTRINT>(Pos1) );
+					const ANSICHAR* StrippedModuleName = RealPos > 0 ? reinterpret_cast<const ANSICHAR*>(RealPos + 1) : SymbolInfo.ModuleName;
+					
+					StackTrace[StackFrameIndex] = FString::Printf(TEXT("%02d - [%s] : [%s] : <%s>:%d"),
+						StackFrameIndex,
+						ANSI_TO_TCHAR(StrippedModuleName),
+						ANSI_TO_TCHAR(SymbolInfo.FunctionName),
+						ANSI_TO_TCHAR(SymbolInfo.Filename),
+						SymbolInfo.LineNumber);
+				}
+
+				FFunctionGraphTask::CreateAndDispatchWhenReady(
+					[
+						StallLogHistory,
+						UniqueID = Params.UniqueID,
+						StackTrace = MoveTemp(StackTrace),
+						StatName = FText::FromStringView(Params.StatName),
+						Now,
+						ThreadID = Params.ThreadId]() mutable
+					{
+						check(IsInGameThread())
+
+						FMessageLog MessageLog("StallLog");
+
+						MessageLog.PerformanceWarning()
+								  ->AddToken(FTextToken::Create(FText::Format(
+									  LOCTEXT("StallDetected", "Stall detected in {0} on {1}"),
+									  StatName, FText::AsDateTime(Now))));
+
+						StallLogHistory->OnStallDetected(UniqueID, Now, StatName.ToString(), MoveTemp(StackTrace), ThreadID);
+					},
+					GET_STATID(STAT_FDelegateGraphTask_StallLogger),
+					nullptr,
+					ENamedThreads::GameThread);
+			});
+	
+	OnStallCompletedDelegate = UE::FStallDetector::StallCompleted.AddLambda(
+			[StallLogHistory = this->StallLogHistory](const UE::FStallCompletedParams& Params)
+			{
+				if (!GEnableStallLogSubsystem)
+				{
+					return;
+				}
+				
+				TSharedRef<FTextToken> MessageToken = FTextToken::Create(FText::Format(
+							NSLOCTEXT("StallDetector", "StallEnded", "Stall ended in {0}: {1} seconds overbudget"),
+							FText::FromStringView(Params.StatName), FText::AsNumber(Params.OverbudgetSeconds)));
+				
+				// Log the end event to the message log. Make sure we're doing that from the game thread.
+				const FGraphEventArray* Prerequisites = nullptr;
+
+				FGraphEventRef LogStallEndAsyncTask_GameThread(
+					FFunctionGraphTask::CreateAndDispatchWhenReady(
+						[
+							MessageTokenIn = MoveTemp(MessageToken),
+							UniqueID = Params.UniqueID,
+							Duration = Params.OverbudgetSeconds,
+							StallLogHistory]()
+						{
+							check(IsInGameThread());
+
+							FMessageLog MessageLog("StallLog");
+
+							MessageLog.PerformanceWarning()
+								->AddToken(MessageTokenIn);
+
+							StallLogHistory->OnStallCompleted(UniqueID, Duration);
+						},
+						GET_STATID(STAT_FDelegateGraphTask_StallLogger),
+						Prerequisites,
+						ENamedThreads::GameThread));
+			}
+		);
+}
+
+void UStallLogSubsystem::UnregisterStallDetectedDelegates()
+{
+	UE::FStallDetector::StallDetected.Remove(OnStallCompletedDelegate);
+	OnStallCompletedDelegate.Reset();
+	
+	UE::FStallDetector::StallDetected.Remove(OnStallCompletedDelegate);
+	OnStallCompletedDelegate.Reset();
+}
+
+namespace UE::Debug
+{
+	static void StallCommand(const TArray<FString>& Arguments)
+	{
+		double SecondsToStall = 2.0;
+		if (Arguments.Num() >= 1)
+		{
+			LexFromString(SecondsToStall, *Arguments[0]);
+		}
+	
+		FFunctionGraphTask::CreateAndDispatchWhenReady(
+			[SecondsToStall]()
+			{
+				SCOPED_NAMED_EVENT_TEXT(TEXT("Fake Stall"), FColor::Red);
+				SCOPE_STALL_COUNTER(FakeStall, 1.0f);
+			
+				const double StartTime = FPlatformTime::Seconds();
+				FPlatformProcess::SleepNoStats(SecondsToStall);
+			
+				while (FPlatformTime::Seconds() - StartTime < SecondsToStall)
+				{
+					// Busy wait the rest if not slept long enough
+				}
+			}, TStatId(), nullptr, ENamedThreads::AnyThread);
+	}
+
+	static void StallAndReportCommand(const TArray<FString>& Arguments)
+	{
+		double SecondsToStall = 2.0;
+		if (Arguments.Num() >= 1)
+		{
+			LexFromString(SecondsToStall, *Arguments[0]);
+		}
+	
+		FFunctionGraphTask::CreateAndDispatchWhenReady(
+			[SecondsToStall]()
+			{
+				SCOPED_NAMED_EVENT_TEXT(TEXT("Fake Stall"), FColor::Red);
+				SCOPE_STALL_REPORTER_ALWAYS(FakeStall, 1.0f);
+			
+				const double StartTime = FPlatformTime::Seconds();
+				FPlatformProcess::SleepNoStats(SecondsToStall);
+			
+				while (FPlatformTime::Seconds() - StartTime < SecondsToStall)
+				{
+					// Busy wait the rest if not slept long enough
+				}
+			}, TStatId(), nullptr, ENamedThreads::AnyThread);
+	}
+
+	static FAutoConsoleCommand CmdEditorStallLoggingStall(
+		TEXT("Editor.Performance.Debug.Stall"),
+		TEXT("Runs a busy loop on the calling thread. Can pass a number of seconds to stall for in parameter (defaults to 2 seconds)."),
+		FConsoleCommandWithArgsDelegate::CreateStatic(&StallCommand)
+	);
+
+	static FAutoConsoleCommand CmdEditorStallLoggingStallAndReport(
+		TEXT("Editor.Performance.Debug.StallAndReport"),
+		TEXT("Runs a busy loop on the calling thread. Can pass a number of seconds to stall for in parameter (defaults to 2 seconds). Will report stall to CRC"),
+		FConsoleCommandWithArgsDelegate::CreateStatic(&StallAndReportCommand)
+	);
+}
+
+#undef LOCTEXT_NAMESPACE
diff --git a/Engine/Plugins/Editor/EditorPerformanceModule/Source/EditorPerformanceModule/Public/EditorPerformanceModule.h b/Engine/Plugins/Editor/EditorPerformanceModule/Source/EditorPerformanceModule/Public/EditorPerformanceModule.h
new file mode 100644
--- /dev/null
+++ b/Engine/Plugins/Editor/EditorPerformanceModule/Source/EditorPerformanceModule/Public/EditorPerformanceModule.h
@@ -0,0 +1,14 @@
+﻿// Copyright Epic Games, Inc. All Rights Reserved.
+
+#pragma once
+
+#include "Modules/ModuleInterface.h"
+
+class EDITORPERFORMANCEMODULE_API IEditorPerformanceModule : public IModuleInterface
+{
+public:
+	static FORCEINLINE IEditorPerformanceModule& Get()
+	{
+		return FModuleManager::LoadModuleChecked<IEditorPerformanceModule>("EditorPerformanceModule");
+	}
+};
\ No newline at end of file
diff --git a/Engine/Plugins/Editor/EditorPerformanceModule/Source/EditorPerformanceModule/Public/StallLogSubsystem/StallLogSubsystem.h b/Engine/Plugins/Editor/EditorPerformanceModule/Source/EditorPerformanceModule/Public/StallLogSubsystem/StallLogSubsystem.h
new file mode 100644
--- /dev/null
+++ b/Engine/Plugins/Editor/EditorPerformanceModule/Source/EditorPerformanceModule/Public/StallLogSubsystem/StallLogSubsystem.h
@@ -0,0 +1,47 @@
+// Copyright Epic Games, Inc. All Rights Reserved.
+
+#pragma once
+
+#include "EditorSubsystem.h"
+
+#include "Containers/ContainersFwd.h"
+#include "Delegates/IDelegateInstance.h"
+#include "Templates/SharedPointerFwd.h"
+
+#include "StallLogSubsystem.generated.h"
+
+class FDelegateHandle;
+class FStallLogHistory;
+struct FSlateBrush;
+class SDockTab;
+
+/**
+ * Subsystem that provides feedback on stall detection
+ */
+UCLASS()
+class EDITORPERFORMANCEMODULE_API UStallLogSubsystem  : public UEditorSubsystem
+{
+	GENERATED_BODY()
+
+public:
+	virtual bool ShouldCreateSubsystem(UObject* Outer) const override;
+	virtual void Initialize(FSubsystemCollectionBase& Collection) override;
+	virtual void Deinitialize() override;
+
+private:
+
+	TSharedRef<SDockTab> CreateStallLogTab(const FSpawnTabArgs& InArgs);
+	const FSlateBrush* GetStatusBarBadgeIcon() const;
+
+	void RegisterStallDetectedDelegates();
+	void UnregisterStallDetectedDelegates();
+
+private:
+
+	FDelegateHandle OnStallDetectedDelegate;
+	FDelegateHandle OnStallCompletedDelegate;
+	TSharedPtr<FStallLogHistory, ESPMode::ThreadSafe> StallLogHistory;
+	
+	TSharedPtr<SWidget> StallLog;
+	TWeakPtr<SDockTab> StallLogTab;
+};
diff --git a/Engine/Source/Runtime/Core/Private/ProfilingDebugging/StallDetector.cpp b/Engine/Source/Runtime/Core/Private/ProfilingDebugging/StallDetector.cpp
--- a/Engine/Source/Runtime/Core/Private/ProfilingDebugging/StallDetector.cpp
+++ b/Engine/Source/Runtime/Core/Private/ProfilingDebugging/StallDetector.cpp
@@ -37,6 +37,9 @@ static int32 InitCount = 0;
 // Sentinel value for invalid timestamp
 static const double InvalidSeconds = -1.0;
 
+UE::FOnStallDetected UE::FStallDetector::StallDetected;
+UE::FOnStallCompleted UE::FStallDetector::StallCompleted;
+
 /**
 * Stall Detector Thread
 **/
@@ -335,6 +338,8 @@ private:
 };
 static ThreadLocalSlotAcquire ThreadLocalSlot;
 
+TAtomic<uint64> UE::FStallDetector::UIDGenerator = 0;
+
 UE::FStallDetector::FStallDetector(FStallDetectorStats& InStats)
 	: Stats(InStats)
 	, Parent(nullptr)
@@ -342,6 +347,8 @@ UE::FStallDetector::FStallDetector(FStallDetectorStats& InStats)
 	, bStarted (false)
 	, bPersistent(false)
 	, Triggered(false)
+	, UniqueID(UIDGenerator.IncrementExchange())
+	
 {
 	// Capture this thread's current stall detector to our parent member
 	Parent = static_cast<FStallDetector*>(FPlatformTLS::GetTlsValue(ThreadLocalSlot.GetSlot()));
@@ -405,19 +412,15 @@ void UE::FStallDetector::Check(bool bFinalCheck, double InCheckSeconds)
 
 	if (Triggered)
 	{
+		// bFinalCheck ensures that only the destructor or CheckAndReset run the final OnStallCompleted
 		if (bFinalCheck)
 		{
-			// this callback allows for tracking the total overrun, once the stalling period is finally done
-			Stats.OnStallCompleted(OverageSeconds);
-
 #if STALL_DETECTOR_DEBUG
 			FString OverageString = FString::Printf(TEXT("[FStallDetector] [%s] Overage of %f\n"), Stats.Name, OverageSeconds);
 			FPlatformMisc::LocalPrint(OverageString.GetCharArray().GetData());
 #endif
-			if (Stats.ReportingMode != EStallDetectorReportingMode::Disabled)
-			{
-				UE_LOG(LogStall, Log, TEXT("Stall detector '%s' complete in %fs (%fs overbudget)"), Stats.Name, DeltaSeconds, OverageSeconds);
-			}
+
+			OnStallCompleted(ThreadId, DeltaSeconds, OverageSeconds);
 		}
 	}
 	else
@@ -432,6 +435,7 @@ void UE::FStallDetector::Check(bool bFinalCheck, double InCheckSeconds)
 				FString OverageString = FString::Printf(TEXT("[FStallDetector] [%s] Triggered at %f\n"), Stats.Name, InCheckSeconds);
 				FPlatformMisc::LocalPrint(OverageString.GetCharArray().GetData());
 #endif
+				
 				OnStallDetected(ThreadId, DeltaSeconds);
 			}
 		}
@@ -459,6 +463,7 @@ void UE::FStallDetector::CheckAndReset()
 
 	Timer.Reset(CheckSeconds, Stats.BudgetSeconds);
 	Triggered = false;
+	UniqueID = UIDGenerator.IncrementExchange();
 }
 
 void UE::FStallDetector::Pause(const double InSeconds)
@@ -545,6 +550,31 @@ void UE::FStallDetector::OnStallDetected(uint32 InThreadId, const double InElaps
 			UE_LOG(LogStall, Log, TEXT("Stall detector '%s' exceeded budget of %fs"), Stats.Name, Stats.BudgetSeconds);
 		}
 	}
+	
+	StallDetected.Broadcast({
+		Stats.Name,
+		ThreadId,
+		UniqueID
+	});
+}
+
+void UE::FStallDetector::OnStallCompleted(uint32 InThreadId, double DeltaSeconds, double OverageSeconds)
+{
+	// this callback allows for tracking the total overrun, once the stalling period is finally done
+	Stats.OnStallCompleted(OverageSeconds);
+			
+	const FStallCompletedParams Params {
+		ThreadId,
+		Stats.Name,
+		OverageSeconds,
+		UniqueID
+	};
+	StallCompleted.Broadcast(Params);
+
+	if (Stats.ReportingMode != EStallDetectorReportingMode::Disabled)
+	{
+		UE_LOG(LogStall, Log, TEXT("Stall detector '%s' complete in %fs (%fs overbudget)"), Stats.Name, DeltaSeconds, OverageSeconds);
+	}
 }
 
 double UE::FStallDetector::Seconds()
diff --git a/Engine/Source/Runtime/Core/Public/ProfilingDebugging/StallDetector.h b/Engine/Source/Runtime/Core/Public/ProfilingDebugging/StallDetector.h
--- a/Engine/Source/Runtime/Core/Public/ProfilingDebugging/StallDetector.h
+++ b/Engine/Source/Runtime/Core/Public/ProfilingDebugging/StallDetector.h
@@ -6,6 +6,8 @@
 #include "Containers/Set.h"
 #include "Containers/UnrealString.h"
 #include "CoreTypes.h"
+#include "Delegates/Delegate.h"
+#include "GenericPlatform/GenericPlatformCrashContext.h"
 #include "HAL/CriticalSection.h"
 #include "HAL/PreprocessorHelpers.h"
 #include "Logging/LogMacros.h"
@@ -33,6 +35,32 @@ CORE_API DECLARE_LOG_CATEGORY_EXTERN(LogStall, Log, All);
 
 namespace UE
 {
+	class FStallDetector;
+
+	struct FStallDetectedParams
+	{
+		/** Name of Stall Detector instance*/
+		FStringView StatName;
+		/** TID of stalled thread */
+		uint32 ThreadId;
+		/** UniqueID of the stalled instance */
+		uint64 UniqueID;
+	};
+	DECLARE_EVENT_OneParam(FStallDetector, FOnStallDetected, const FStallDetectedParams&);
+	
+	struct FStallCompletedParams
+	{
+		/** TID of stalled thread */
+		uint32 ThreadId;
+		/** Name of Stall Detector instance*/
+		FStringView StatName;
+		/** Amount of time StallDetector instance stalled  */
+		double OverbudgetSeconds;
+		/** UniqueID of the stalled instance */
+		uint64 UniqueID;
+	};
+	DECLARE_EVENT_OneParam(FStallDetector, FOnStallCompleted, const FStallCompletedParams&);
+
 	/**
 	* The reporting behavior of the detector
 	**/
@@ -60,7 +88,7 @@ namespace UE
 		~FStallDetectorStats();
 
 		/**
-		* Stall dettector scope has ended, record the interval
+		* Stall detector scope has ended, record the interval
 		*/
 		void OnStallCompleted(double InOverageSeconds);
 
@@ -263,7 +291,24 @@ namespace UE
 		**/
 		static bool IsRunning();
 
+		/**
+		 * Event that triggers when we detect that a stall in any instance is occuring
+		 * Note: It is possible for multiple instances to stall at the same time,
+		 * delegates should take this into account 
+		 */
+		static FOnStallDetected StallDetected;
+
+		/**
+		 * Event that triggers when a stall has stopped
+		 * Note: It is possible for multiple instances to complete their stall at the same time,
+		 * delegates should take this into account
+		 */
+		static FOnStallCompleted StallCompleted;
+
 	private:
+		/** Monotonically increasing counter to generate a unique ID per stall detector */
+		static TAtomic<uint64> UIDGenerator;
+		
 		// The time provided, the deadline
 		FStallDetectorStats& Stats;
 
@@ -285,9 +330,15 @@ namespace UE
 		// Track the triggered state, atomic to mitigate foreground and background thread checking
 		std::atomic<bool> Triggered;
 
+		// Unique ID given to a StallDetector
+		uint64 UniqueID;
+
 		// The pinch point for tracking counts and emitting a report for a triggered hitch detector
 		void OnStallDetected(uint32 InThreadId, const double InElapsedSeconds);
 
+		// Called when the thread that triggered the detector is no longer stalling
+		void OnStallCompleted(uint32 InThreadId, const double DeltaSeconds, const double OverageSeconds);
+
 		// Guards access to the instances container from multiple threads
 		static FCriticalSection InstancesSection;
 
