Fix Non-Unity build errors in EditorPerformanceModule

diff --git a/Engine/Plugins/Editor/EditorPerformanceModule/Source/EditorPerformanceModule/Private/StallLogSubsystem/StallLogSubsystem.cpp b/Engine/Plugins/Editor/EditorPerformanceModule/Source/EditorPerformanceModule/Private/StallLogSubsystem/StallLogSubsystem.cpp
--- a/Engine/Plugins/Editor/EditorPerformanceModule/Source/EditorPerformanceModule/Private/StallLogSubsystem/StallLogSubsystem.cpp
+++ b/Engine/Plugins/Editor/EditorPerformanceModule/Source/EditorPerformanceModule/Private/StallLogSubsystem/StallLogSubsystem.cpp
@@ -4,6 +4,7 @@
 
 #include "Async/TaskGraphInterfaces.h"
 #include "Containers/Array.h"
+#include "Framework/Application/SlateApplication.h"
 #include "Framework/Notifications/NotificationManager.h"
 #include "GenericPlatform/GenericPlatformCrashContext.h"
 #include "HAL/IConsoleManager.h"
@@ -17,9 +18,12 @@
 #include "Stats/Stats2.h"
 #include "ToolMenus.h"
 #include "Widgets/Docking/SDockTab.h"
+#include "Widgets/Images/SImage.h"
+#include "Widgets/Input/SButton.h"
 #include "Widgets/Layout/SHeader.h"
 #include "Widgets/Layout/SSpacer.h"
 #include "Widgets/Notifications/SNotificationList.h"
+#include "Widgets/Views/SListView.h"
 #include "Widgets/Views/STableRow.h"
 #include "WorkspaceMenuStructure.h"
 #include "WorkspaceMenuStructureModule.h"
diff --git a/Engine/Plugins/Editor/EditorPerformanceModule/Source/EditorPerformanceModule/Public/EditorPerformanceModule.h b/Engine/Plugins/Editor/EditorPerformanceModule/Source/EditorPerformanceModule/Public/EditorPerformanceModule.h
--- a/Engine/Plugins/Editor/EditorPerformanceModule/Source/EditorPerformanceModule/Public/EditorPerformanceModule.h
+++ b/Engine/Plugins/Editor/EditorPerformanceModule/Source/EditorPerformanceModule/Public/EditorPerformanceModule.h
@@ -3,6 +3,7 @@
 #pragma once
 
 #include "Modules/ModuleInterface.h"
+#include "Modules/ModuleManager.h"
 
 class EDITORPERFORMANCEMODULE_API IEditorPerformanceModule : public IModuleInterface
 {
diff --git a/Engine/Plugins/Editor/EditorPerformanceModule/Source/EditorPerformanceModule/Public/StallLogSubsystem/StallLogSubsystem.h b/Engine/Plugins/Editor/EditorPerformanceModule/Source/EditorPerformanceModule/Public/StallLogSubsystem/StallLogSubsystem.h
--- a/Engine/Plugins/Editor/EditorPerformanceModule/Source/EditorPerformanceModule/Public/StallLogSubsystem/StallLogSubsystem.h
+++ b/Engine/Plugins/Editor/EditorPerformanceModule/Source/EditorPerformanceModule/Public/StallLogSubsystem/StallLogSubsystem.h
@@ -6,14 +6,16 @@
 
 #include "Containers/ContainersFwd.h"
 #include "Delegates/IDelegateInstance.h"
-#include "Templates/SharedPointerFwd.h"
+#include "Templates/SharedPointer.h"
 
 #include "StallLogSubsystem.generated.h"
 
 class FDelegateHandle;
+class FSpawnTabArgs;
 class FStallLogHistory;
 struct FSlateBrush;
 class SDockTab;
+class SWidget;
 
 /**
  * Subsystem that provides feedback on stall detection
