Fix mac editor compile
* StallLogSubsystem depends on StallDetector, StallDetector does not define anything on mac

diff --git a/Engine/Plugins/Editor/EditorPerformanceModule/Source/EditorPerformanceModule/Private/StallLogSubsystem/StallLogSubsystem.cpp b/Engine/Plugins/Editor/EditorPerformanceModule/Source/EditorPerformanceModule/Private/StallLogSubsystem/StallLogSubsystem.cpp
--- a/Engine/Plugins/Editor/EditorPerformanceModule/Source/EditorPerformanceModule/Private/StallLogSubsystem/StallLogSubsystem.cpp
+++ b/Engine/Plugins/Editor/EditorPerformanceModule/Source/EditorPerformanceModule/Private/StallLogSubsystem/StallLogSubsystem.cpp
@@ -552,6 +552,11 @@ const FSlateBrush* UStallLogSubsystem::GetStatusBarBadgeIcon() const
 
 void UStallLogSubsystem::RegisterStallDetectedDelegates()
 {
+	// Only compile this if StallDetector exists
+	// Note: This is not preferred, however the implementation relies on StallDetector and definitions are stripped
+	// on many configurations
+#if STALL_DETECTOR
+	
 	OnStallDetectedDelegate = UE::FStallDetector::StallDetected.AddLambda(
 			[StallLogHistory = this->StallLogHistory](const UE::FStallDetectedParams& Params)
 			{
@@ -659,15 +664,18 @@ void UStallLogSubsystem::RegisterStallDetectedDelegates()
 						ENamedThreads::GameThread));
 			}
 		);
+#endif // STALL_DETECTOR
 }
 
 void UStallLogSubsystem::UnregisterStallDetectedDelegates()
 {
+#if STALL_DETECTOR
 	UE::FStallDetector::StallDetected.Remove(OnStallCompletedDelegate);
 	OnStallCompletedDelegate.Reset();
 	
 	UE::FStallDetector::StallDetected.Remove(OnStallCompletedDelegate);
 	OnStallCompletedDelegate.Reset();
+#endif // STALL_DETECTOR
 }
 
 namespace UE::Debug
