Remove logging from EditorPerformanceModule
* Appears to conflict with Asset Validation - Asset Validation picks up the warning that the stall detector generates and causes false failures.

diff --git a/Engine/Plugins/Editor/EditorPerformanceModule/Source/EditorPerformanceModule/Private/StallLogSubsystem/StallLogSubsystem.cpp b/Engine/Plugins/Editor/EditorPerformanceModule/Source/EditorPerformanceModule/Private/StallLogSubsystem/StallLogSubsystem.cpp
--- a/Engine/Plugins/Editor/EditorPerformanceModule/Source/EditorPerformanceModule/Private/StallLogSubsystem/StallLogSubsystem.cpp
+++ b/Engine/Plugins/Editor/EditorPerformanceModule/Source/EditorPerformanceModule/Private/StallLogSubsystem/StallLogSubsystem.cpp
@@ -613,13 +613,6 @@ void UStallLogSubsystem::RegisterStallDetectedDelegates()
 					{
 						check(IsInGameThread())
 
-						FMessageLog MessageLog("StallLog");
-
-						MessageLog.PerformanceWarning()
-								  ->AddToken(FTextToken::Create(FText::Format(
-									  LOCTEXT("StallDetected", "Stall detected in {0} on {1}"),
-									  StatName, FText::AsDateTime(Now))));
-
 						StallLogHistory->OnStallDetected(UniqueID, Now, StatName.ToString(), MoveTemp(StackTrace), ThreadID);
 					},
 					GET_STATID(STAT_FDelegateGraphTask_StallLogger),
@@ -635,17 +628,12 @@ void UStallLogSubsystem::RegisterStallDetectedDelegates()
 					return;
 				}
 				
-				TSharedRef<FTextToken> MessageToken = FTextToken::Create(FText::Format(
-							NSLOCTEXT("StallDetector", "StallEnded", "Stall ended in {0}: {1} seconds overbudget"),
-							FText::FromStringView(Params.StatName), FText::AsNumber(Params.OverbudgetSeconds)));
-				
 				// Log the end event to the message log. Make sure we're doing that from the game thread.
 				const FGraphEventArray* Prerequisites = nullptr;
 
 				FGraphEventRef LogStallEndAsyncTask_GameThread(
 					FFunctionGraphTask::CreateAndDispatchWhenReady(
 						[
-							MessageTokenIn = MoveTemp(MessageToken),
 							UniqueID = Params.UniqueID,
 							Duration = Params.OverbudgetSeconds,
 							StallLogHistory]()
@@ -654,9 +642,6 @@ void UStallLogSubsystem::RegisterStallDetectedDelegates()
 
 							FMessageLog MessageLog("StallLog");
 
-							MessageLog.PerformanceWarning()
-								->AddToken(MessageTokenIn);
-
 							StallLogHistory->OnStallCompleted(UniqueID, Duration);
 						},
 						GET_STATID(STAT_FDelegateGraphTask_StallLogger),
