Defer symbolification of BackTrace captured by [REDACTED_PERSON]
* Further mitigation to avoid Asset Validation errors. Another issue that came up was that if symbols are not found, then a warning is generated by the platform layer which Asset Validation picks up and assumes is a failure of the asset.
* Mitigation strategy is that Asset Validation is a SlowTask and blocks the user from interfacing with the editor anyways.  Stall log should be able to do a backtrace without upsetting the Asset Validator.
* This is technically safer to do on Windows and also doesn't pull in large PDBs unnecessarily.
* Symbol resolving is off by default, but can be enabled with a CVar.

diff --git a/Engine/Plugins/Editor/EditorPerformanceModule/Source/EditorPerformanceModule/Private/StallLogSubsystem/StallLogSubsystem.cpp b/Engine/Plugins/Editor/EditorPerformanceModule/Source/EditorPerformanceModule/Private/StallLogSubsystem/StallLogSubsystem.cpp
--- a/Engine/Plugins/Editor/EditorPerformanceModule/Source/EditorPerformanceModule/Private/StallLogSubsystem/StallLogSubsystem.cpp
+++ b/Engine/Plugins/Editor/EditorPerformanceModule/Source/EditorPerformanceModule/Private/StallLogSubsystem/StallLogSubsystem.cpp
@@ -12,6 +12,7 @@
 #include "Logging/MessageLog.h"
 #include "MessageLogInitializationOptions.h"
 #include "MessageLogModule.h"
+#include "Modules/BuildVersion.h"
 #include "Modules/ModuleManager.h"
 #include "ProfilingDebugging/StallDetector.h"
 #include "SlateOptMacros.h"
@@ -30,13 +31,25 @@
 
 #define LOCTEXT_NAMESPACE "StallLogSubsystem"
 
-static bool GEnableStallLogSubsystem = true;
-static FAutoConsoleVariableRef CCvarEnableStallLogging(
+namespace 
+{
+	bool EnableStallLogSubsystem = true;
+	bool EnableStallLogStackTracing = false;
+}
+
+static FAutoConsoleVariableRef CVarEnableStallLogging(
 	TEXT("Editor.StallLogger.Enable"),
-	GEnableStallLogSubsystem,
+	EnableStallLogSubsystem,
 	TEXT("Whether the editor stall logger subsystem is enabled."),
 	ECVF_Default);
 
+
+static FAutoConsoleVariableRef CVarEnableStallStackTracing(
+	TEXT("Editor.StallLogger.EnableStackTrace"),
+	EnableStallLogStackTracing,
+	TEXT("When enabled, stall logs will do a stacktrace to get human-readable function history"),
+	ECVF_Default);
+
 /**
  * Metadata for each detected stall of the application
  */
@@ -44,19 +57,19 @@ class FStallLogItem : public TSharedFromThis<FStallLogItem>
 {
 public:
 	FStallLogItem() = default;
-	FStallLogItem(FString InLocation, float InDurationSeconds, FDateTime InTime, TArray<FString> InStackTrace);
+	FStallLogItem(FString InLocation, float InDurationSeconds, FDateTime InTime, TArray<uint64> InBackTrace);
 
 	FString Location;
 	float DurationSeconds;
 	FDateTime Time;
-	TArray<FString> StackTrace;
+	TArray<uint64> BackTrace;
 };
 
-FStallLogItem::FStallLogItem(FString InLocation, float InDurationSeconds, FDateTime InTime, TArray<FString> InStackTrace)
+FStallLogItem::FStallLogItem(FString InLocation, float InDurationSeconds, FDateTime InTime, TArray<uint64> InBackTrace)
 	: Location(MoveTemp(InLocation))
 	, DurationSeconds(InDurationSeconds)
 	, Time(InTime)
-	, StackTrace(MoveTemp(InStackTrace))
+	, BackTrace(MoveTemp(InBackTrace))
 {}
 
 using FStallLogItemPtr = TSharedPtr<FStallLogItem>;
@@ -68,7 +81,7 @@ using FStallLogItemPtr = TSharedPtr<FStallLogItem>;
 class FStallLogHistory
 {
 public:
-	void OnStallDetected(uint64 UniqueID, FDateTime InDetectTime, FStringView StatName, TArray<FString> StackTrace, uint32 ThreadID);
+	void OnStallDetected(uint64 UniqueID, FDateTime InDetectTime, FStringView StatName, TArray<uint64> Backtrace, uint32 ThreadID);
 	void OnStallCompleted(uint64 UniqueID, double InDurationSeconds);
 	
 	void ClearStallLog();
@@ -80,7 +93,7 @@ private:
 	{
 		FDateTime DetectTime;
 		FString StatName;
-		TArray<FString> StackTrace;
+		TArray<uint64> Backtrace;
 		uint32 ThreadID;
 	};
 
@@ -184,7 +197,66 @@ namespace
 					.OnClicked_Lambda([this]()
 					{
 						// Build a string of the stall information and put it on the clipboard
-						FPlatformApplicationMisc::ClipboardCopy(*FString::Join(StallLogItem->StackTrace, TEXT("\n")));
+						TStringBuilder<32768> Clipboard;
+						
+						FBuildVersion BuildVersion;
+						if (FBuildVersion::TryRead(FBuildVersion::GetDefaultFileName(), BuildVersion))
+						{
+							Clipboard.Appendf(TEXT("Engine Version: %s\n"), *BuildVersion.GetEngineVersion().ToString());
+						}
+						else
+						{
+							Clipboard.Append(TEXT("Engine Version: <unknown>\n"));
+						}
+						Clipboard.Appendf(TEXT("Stall Detector: %s\n"), *StallLogItem->Location);
+						Clipboard.Appendf(TEXT("Stall Duration: %03f\n"), StallLogItem->DurationSeconds);
+						Clipboard.Appendf(TEXT("Stall Time: %s\n"), *StallLogItem->Time.ToString());
+						
+						if (!EnableStallLogStackTracing)
+						{
+							Clipboard.Append(TEXT("BackTrace\n=========\n"));
+							for (int32 StackFrameIndex = 0; StackFrameIndex < StallLogItem->BackTrace.Num(); ++StackFrameIndex)
+							{
+								const uint64 ProgramCounter = StallLogItem->BackTrace[StackFrameIndex];
+								Clipboard.Appendf(TEXT("%02d: [0x%p]\n"), StackFrameIndex, ProgramCounter);
+							}
+							Clipboard.Append(TEXT("\nUse CMD: \"Editor.StallLogger.EnableStackTrace = 1\" to enable more detailed stacktrace\"\n"));
+						}
+						else // !GEnableStallLogStackTracing
+						{
+							Clipboard.Append(TEXT("\nStackTrace\n==========\n"));
+							const bool StackWalkingInitialized = FPlatformStackWalk::InitStackWalking();
+
+							if (StackWalkingInitialized)
+							{
+								for (int32 StackFrameIndex = 0; StackFrameIndex < StallLogItem->BackTrace.Num(); ++StackFrameIndex)
+								{
+									FProgramCounterSymbolInfo SymbolInfo;
+									const uint64 ProgramCounter = StallLogItem->BackTrace[StackFrameIndex];
+									FPlatformStackWalk::ProgramCounterToSymbolInfo(ProgramCounter, SymbolInfo);
+
+									// Strip module path.
+									const ANSICHAR* Pos0 = FCStringAnsi::Strrchr( SymbolInfo.ModuleName, '\\' );
+									const ANSICHAR* Pos1 = FCStringAnsi::Strrchr( SymbolInfo.ModuleName, '/' );
+									const UPTRINT RealPos = FMath::Max(reinterpret_cast<UPTRINT>(Pos0), reinterpret_cast<UPTRINT>(Pos1) );
+									const ANSICHAR* StrippedModuleName = RealPos > 0 ? reinterpret_cast<const ANSICHAR*>(RealPos + 1) : SymbolInfo.ModuleName;
+											
+									Clipboard.Appendf(TEXT("%02d: [0x%p] [%s] : [%s] : <%s>:%d\n"),
+										StackFrameIndex,
+										ProgramCounter,
+										ANSI_TO_TCHAR(StrippedModuleName),
+										ANSI_TO_TCHAR(SymbolInfo.FunctionName),
+										ANSI_TO_TCHAR(SymbolInfo.Filename),
+										SymbolInfo.LineNumber);
+								}
+							}
+							else
+							{
+								Clipboard.Append(TEXT("- Failed to initialize StackWalking\n"));
+							}
+						}
+						
+						FPlatformApplicationMisc::ClipboardCopy(*Clipboard);
 
 						FNotificationInfo Info(LOCTEXT("StallLogInfoCopied", "Copied to clipboard"));
 						Info.ExpireDuration = 2.0f;
@@ -399,14 +471,14 @@ void FStallLogHistory::OnStallDetected(
 	uint64 UniqueID,
 	FDateTime InDetectTime,
 	FStringView StatName,
-	TArray<FString> StackTrace,
+	TArray<uint64> Backtrace,
 	uint32 ThreadID)
 {
 	checkf(IsInGameThread(), TEXT("Can only be run on GameThread"));
 	FInFlightStall InFlightStall;
 	InFlightStall.DetectTime = InDetectTime;
 	InFlightStall.StatName = StatName;
-	InFlightStall.StackTrace = MoveTemp(StackTrace);
+	InFlightStall.Backtrace = MoveTemp(Backtrace);
 	InFlightStall.ThreadID = ThreadID;
 
 	InFlightStalls.Emplace(UniqueID, MoveTemp(InFlightStall));
@@ -426,7 +498,7 @@ void FStallLogHistory::OnStallCompleted(uint64 UniqueID, double InDurationSecond
 			MoveTemp(InFlightStall->StatName), 
 			InDurationSeconds,
 			InFlightStall->DetectTime,
-			MoveTemp(InFlightStall->StackTrace));
+			MoveTemp(InFlightStall->Backtrace));
 		StallLogs.Emplace(StallLogItem);
 
 		InFlightStalls.Remove(UniqueID);
@@ -560,7 +632,7 @@ void UStallLogSubsystem::RegisterStallDetectedDelegates()
 	OnStallDetectedDelegate = UE::FStallDetector::StallDetected.AddLambda(
 			[StallLogHistory = this->StallLogHistory](const UE::FStallDetectedParams& Params)
 			{
-				if (!GEnableStallLogSubsystem)
+				if (!EnableStallLogSubsystem)
 				{
 					return;
 				}
@@ -580,40 +652,18 @@ void UStallLogSubsystem::RegisterStallDetectedDelegates()
 
 				Backtrace.SetNum(StackDepth);
 
-				TArray<FString> StackTrace;
-				StackTrace.SetNum(StackDepth);
-
-				for (uint32 StackFrameIndex = 0; StackFrameIndex < StackDepth; ++StackFrameIndex)
-				{
-					FProgramCounterSymbolInfo SymbolInfo;
-					FPlatformStackWalk::ProgramCounterToSymbolInfo(Backtrace[StackFrameIndex], SymbolInfo);
-
-					// Strip module path.
-					const ANSICHAR* Pos0 = FCStringAnsi::Strrchr( SymbolInfo.ModuleName, '\\' );
-					const ANSICHAR* Pos1 = FCStringAnsi::Strrchr( SymbolInfo.ModuleName, '/' );
-					const UPTRINT RealPos = FMath::Max(reinterpret_cast<UPTRINT>(Pos0), reinterpret_cast<UPTRINT>(Pos1) );
-					const ANSICHAR* StrippedModuleName = RealPos > 0 ? reinterpret_cast<const ANSICHAR*>(RealPos + 1) : SymbolInfo.ModuleName;
-					
-					StackTrace[StackFrameIndex] = FString::Printf(TEXT("%02d - [%s] : [%s] : <%s>:%d"),
-						StackFrameIndex,
-						ANSI_TO_TCHAR(StrippedModuleName),
-						ANSI_TO_TCHAR(SymbolInfo.FunctionName),
-						ANSI_TO_TCHAR(SymbolInfo.Filename),
-						SymbolInfo.LineNumber);
-				}
-
 				FFunctionGraphTask::CreateAndDispatchWhenReady(
 					[
 						StallLogHistory,
 						UniqueID = Params.UniqueID,
-						StackTrace = MoveTemp(StackTrace),
+						Backtrace = MoveTemp(Backtrace),
 						StatName = FText::FromStringView(Params.StatName),
 						Now,
 						ThreadID = Params.ThreadId]() mutable
 					{
 						check(IsInGameThread())
 
-						StallLogHistory->OnStallDetected(UniqueID, Now, StatName.ToString(), MoveTemp(StackTrace), ThreadID);
+						StallLogHistory->OnStallDetected(UniqueID, Now, StatName.ToString(), MoveTemp(Backtrace), ThreadID);
 					},
 					GET_STATID(STAT_FDelegateGraphTask_StallLogger),
 					nullptr,
@@ -623,7 +673,7 @@ void UStallLogSubsystem::RegisterStallDetectedDelegates()
 	OnStallCompletedDelegate = UE::FStallDetector::StallCompleted.AddLambda(
 			[StallLogHistory = this->StallLogHistory](const UE::FStallCompletedParams& Params)
 			{
-				if (!GEnableStallLogSubsystem)
+				if (!EnableStallLogSubsystem)
 				{
 					return;
 				}
