Update stall log subsystem for stall detector changes

diff --git a/Engine/Plugins/Editor/EditorPerformanceModule/Source/EditorPerformanceModule/Private/StallLogSubsystem/StallLogSubsystem.cpp b/Engine/Plugins/Editor/EditorPerformanceModule/Source/EditorPerformanceModule/Private/StallLogSubsystem/StallLogSubsystem.cpp
--- a/Engine/Plugins/Editor/EditorPerformanceModule/Source/EditorPerformanceModule/Private/StallLogSubsystem/StallLogSubsystem.cpp
+++ b/Engine/Plugins/Editor/EditorPerformanceModule/Source/EditorPerformanceModule/Private/StallLogSubsystem/StallLogSubsystem.cpp
@@ -9,6 +9,7 @@
 #include "GenericPlatform/GenericPlatformCrashContext.h"
 #include "HAL/IConsoleManager.h"
 #include "HAL/PlatformApplicationMisc.h"
+#include "HAL/ThreadManager.h"
 #include "Logging/MessageLog.h"
 #include "MessageLogInitializationOptions.h"
 #include "MessageLogModule.h"
@@ -34,7 +35,7 @@
 namespace 
 {
 	bool EnableStallLogSubsystem = true;
-	bool EnableStallLogStackTracing = false;
+	bool EnableStallLogStackTracing = true;
 }
 
 static FAutoConsoleVariableRef CVarEnableStallLogging(
@@ -57,21 +58,64 @@ class FStallLogItem : public TSharedFromThis<FStallLogItem>
 {
 public:
 	FStallLogItem() = default;
-	FStallLogItem(FString InLocation, float InDurationSeconds, FDateTime InTime, TArray<uint64> InBackTrace);
+
+	// Constructor for OnStallDetected firing first
+	FStallLogItem(FDateTime InDetectTime, TConstArrayView<uint64> InBackTrace)
+		: Location()
+		, ThreadName()
+		, DurationSeconds()
+		, Time(InDetectTime)
+		, BackTrace(InBackTrace)
+	{
+	}
+
+	// Constructor for OnStallDetected firing second
+	FStallLogItem(FStallLogItem&& Old, FDateTime InDetectTime, TConstArrayView<uint64> InBackTrace)
+		: Location(MoveTemp(Old.Location))
+		, ThreadName(MoveTemp(Old.ThreadName))
+		, DurationSeconds(Old.DurationSeconds)
+		, Time(InDetectTime)
+		, BackTrace(InBackTrace)
+	{
+	}
+
+	// Constructor for OnStallComplete firing first
+	FStallLogItem(FStringView Location, FString InThreadName, double DurationSeconds)
+		: Location(Location)
+		, ThreadName(InThreadName)
+		, DurationSeconds(DurationSeconds)
+		, Time()
+		, BackTrace()
+	{
+	}
+
+	// Constructor for OnStallComplete firing second
+	FStallLogItem(FStallLogItem&& Old, FStringView Location, FString InThreadName, double DurationSeconds)
+		: Location(Location)
+		, ThreadName(InThreadName)
+		, DurationSeconds(DurationSeconds)
+		, Time(Old.Time)
+		, BackTrace(MoveTemp(Old.BackTrace))
+	{
+	}
+
+	// Constructor for OnstallComplte firing without background detection
+	FStallLogItem(FStringView InLocation, FString ThreadName, float InDurationSeconds, FDateTime InTime)
+		: Location(InLocation)
+		, ThreadName(MoveTemp(ThreadName))
+		, DurationSeconds(InDurationSeconds)
+		, Time(InTime)
+		, BackTrace()
+	{
+	}
 
 	FString Location;
+	FString ThreadName;
 	float DurationSeconds;
 	FDateTime Time;
 	TArray<uint64> BackTrace;
 };
 
-FStallLogItem::FStallLogItem(FString InLocation, float InDurationSeconds, FDateTime InTime, TArray<uint64> InBackTrace)
-	: Location(MoveTemp(InLocation))
-	, DurationSeconds(InDurationSeconds)
-	, Time(InTime)
-	, BackTrace(MoveTemp(InBackTrace))
-{}
-
 using FStallLogItemPtr = TSharedPtr<FStallLogItem>;
 
 /**
@@ -81,23 +125,17 @@ using FStallLogItemPtr = TSharedPtr<FStallLogItem>;
 class FStallLogHistory
 {
 public:
-	void OnStallDetected(uint64 UniqueID, FDateTime InDetectTime, FStringView StatName, TArray<uint64> Backtrace, uint32 ThreadID);
-	void OnStallCompleted(uint64 UniqueID, double InDurationSeconds);
-	
+	void OnStallDetected(uint64 UniqueID, FDateTime InDetectTime, TConstArrayView<uint64> Backtrace);
+	void OnStallCompleted(uint32 ThreadID, FStringView StatName, uint64 UniqueID, FDateTime InCompletedTime, double InDurationSeconds, bool bWasDetectedWithCallstack);
+
 	void ClearStallLog();
 	
 	const TArray<FStallLogItemPtr>& GetStallLog() const;
 
 private:
-	struct FInFlightStall
-	{
-		FDateTime DetectTime;
-		FString StatName;
-		TArray<uint64> Backtrace;
-		uint32 ThreadID;
-	};
-
-	TMap<uint64, FInFlightStall> InFlightStalls;
+	// Stalls for which we have received only one of the necessary callbacks, not yet added to StallLogs
+	TMap<uint64, FStallLogItem> InFlightStalls;
+	// Completed stalls with full information
 	TArray<FStallLogItemPtr> StallLogs;
 };
 
@@ -118,6 +156,7 @@ namespace
 	const FName StallLogTabName = FName(TEXT("StallLogTab"));
 
 	const FName ColumnName_Location = FName(TEXT("Location"));
+	const FName ColumnName_Thread = FName(TEXT("Thread"));
 	const FName ColumnName_Duration = FName(TEXT("Duration"));
 	const FName ColumnName_Time = FName(TEXT("Time"));
 	const FName ColumnName_Copy = FName(TEXT("Copy"));
@@ -156,6 +195,15 @@ namespace
 							.Text(FText::FromString(StallLogItem->Location))
 					];
 			}
+			else if (ColumnName == ColumnName_Thread)
+			{
+				return SNew(SBox)
+					.Padding(FMargin(4.0f, 0.0f))
+					.VAlign(VAlign_Center)
+						[SNew(STextBlock)
+								.ColorAndOpacity(FSlateColor::UseForeground())
+								.Text(FText::FromString(StallLogItem->ThreadName))];
+			}
 			else if (ColumnName == ColumnName_Duration)
 			{
 				return SNew(SBox)
@@ -314,7 +362,7 @@ namespace
 		};
 
 		ClearLogDelegate = InArgs._OnClearLog;
-		
+
 		ChildSlot
 		.Padding(3)
 		[
@@ -351,6 +399,10 @@ namespace
 										.DefaultLabel(LOCTEXT("StallLogColumnHeader_StallDetectorName", "Stall Detector Name"))
 										.FillWidth(0.3f)
 
+									+ SHeaderRow::Column(ColumnName_Thread)
+										  .DefaultLabel(LOCTEXT("StallLogColumnHeader_ThreadName", "Thread Name"))
+										  .FillWidth(0.3f)
+
 									+ SHeaderRow::Column(ColumnName_Duration)
 										.DefaultLabel(LOCTEXT("StallLogColumnHeader_Duration", "Duration"))
 										.FillWidth(0.3f)
@@ -397,6 +449,7 @@ namespace
 				]
 			]
 		];
+
 	}
 	
 	class SStallLogStatusBarWidget : public SCompoundWidget
@@ -470,38 +523,52 @@ namespace
 void FStallLogHistory::OnStallDetected(
 	uint64 UniqueID,
 	FDateTime InDetectTime,
-	FStringView StatName,
-	TArray<uint64> Backtrace,
-	uint32 ThreadID)
+	TConstArrayView<uint64> Backtrace)
 {
 	checkf(IsInGameThread(), TEXT("Can only be run on GameThread"));
-	FInFlightStall InFlightStall;
-	InFlightStall.DetectTime = InDetectTime;
-	InFlightStall.StatName = StatName;
-	InFlightStall.Backtrace = MoveTemp(Backtrace);
-	InFlightStall.ThreadID = ThreadID;
 
-	InFlightStalls.Emplace(UniqueID, MoveTemp(InFlightStall));
+	FStallLogItem Existing;
+	if (InFlightStalls.RemoveAndCopyValue(UniqueID, Existing))
+	{
+		// Already handled completed event, just add back trace & detect time
+		StallLogs.Emplace(MakeShared<FStallLogItem>(MoveTemp(Existing), InDetectTime, Backtrace));
+	}
+	else
+	{
+		// Wait for completed event to finish with full duration
+		InFlightStalls.Add(UniqueID, FStallLogItem(InDetectTime, Backtrace));
+	}
 }
 
-void FStallLogHistory::OnStallCompleted(uint64 UniqueID, double InDurationSeconds)
+void FStallLogHistory::OnStallCompleted(uint32 ThreadID, FStringView StatName, uint64 UniqueID, FDateTime InCompletedTime, double InDurationSeconds, bool bWasDetectedWithBacktrace)
 {
 	checkf(IsInGameThread(), TEXT("Can only be run on GameThread"));
-	
-	// Find the completed stall and only add it to the history if found
-	FInFlightStall* InFlightStall = InFlightStalls.Find(UniqueID);
-	ensure(InFlightStall != nullptr);
-	
-	if (InFlightStall != nullptr)
+
+	FString ThreadName = FThreadManager::GetThreadName(ThreadID);
+	if (bWasDetectedWithBacktrace)
+	{
+		FStallLogItem Existing;
+		if (InFlightStalls.RemoveAndCopyValue(UniqueID, Existing))
+		{
+			// Fill in final duration and add to list
+			// Existing->DurationSeconds = InDurationSeconds;
+			StallLogs.Emplace(MakeShared<FStallLogItem>(MoveTemp(Existing), StatName, MoveTemp(ThreadName), InDurationSeconds));
+		}
+		else
+		{
+			// Still waiting for callstack, add full duration
+			InFlightStalls.Add(UniqueID,
+				FStallLogItem(StatName, FThreadManager::GetThreadName(ThreadID), InDurationSeconds));
+		}
+	}
+	else
 	{
-		TSharedPtr<FStallLogItem> StallLogItem = MakeShared<FStallLogItem>(
-			MoveTemp(InFlightStall->StatName), 
+		// No backtrace for this event
+		StallLogs.Emplace(MakeShared<FStallLogItem>(
+			StatName,
+			MoveTemp(ThreadName),
 			InDurationSeconds,
-			InFlightStall->DetectTime,
-			MoveTemp(InFlightStall->Backtrace));
-		StallLogs.Emplace(StallLogItem);
-
-		InFlightStalls.Remove(UniqueID);
+			InCompletedTime));
 	}
 }
 
@@ -642,28 +709,17 @@ void UStallLogSubsystem::RegisterStallDetectedDelegates()
 				// Add a bookmark to Insights when a stall is detected
 				// Helps understand the context of the stall on the thread timeline
 				TRACE_BOOKMARK(TEXT("Stall [%s]"), *FString(Params.StatName));
-				
-				// Grab a stacktrace of the stalled thread
-				const uint32 MaxStackDepth = 64;
-				TArray<uint64> Backtrace;
-				Backtrace.SetNumUninitialized(MaxStackDepth);
-				const uint32 StackDepth = FPlatformStackWalk::CaptureThreadStackBackTrace(
-					Params.ThreadId, Backtrace.GetData(), MaxStackDepth);
-
-				Backtrace.SetNum(StackDepth);
 
 				FFunctionGraphTask::CreateAndDispatchWhenReady(
-					[
-						StallLogHistory,
+					[StallLogHistory,
 						UniqueID = Params.UniqueID,
-						Backtrace = MoveTemp(Backtrace),
+						Backtrace = Params.Backtrace,
 						StatName = FText::FromStringView(Params.StatName),
 						Now,
-						ThreadID = Params.ThreadId]() mutable
-					{
+						ThreadID = Params.ThreadId]() mutable {
 						check(IsInGameThread())
 
-						StallLogHistory->OnStallDetected(UniqueID, Now, StatName.ToString(), MoveTemp(Backtrace), ThreadID);
+							StallLogHistory->OnStallDetected(UniqueID, Now, Backtrace);
 					},
 					GET_STATID(STAT_FDelegateGraphTask_StallLogger),
 					nullptr,
@@ -683,16 +739,20 @@ void UStallLogSubsystem::RegisterStallDetectedDelegates()
 
 				FGraphEventRef LogStallEndAsyncTask_GameThread(
 					FFunctionGraphTask::CreateAndDispatchWhenReady(
-						[
-							UniqueID = Params.UniqueID,
-							Duration = Params.OverbudgetSeconds,
-							StallLogHistory]()
-						{
+						[Params,
+							Now = FDateTime::Now(),
+							StallLogHistory]() {
 							check(IsInGameThread());
 
 							FMessageLog MessageLog("StallLog");
 
-							StallLogHistory->OnStallCompleted(UniqueID, Duration);
+							StallLogHistory->OnStallCompleted(
+								Params.ThreadId,
+								Params.StatName,
+								Params.UniqueID,
+								Now,
+								Params.BudgetSeconds + Params.OverbudgetSeconds,
+								Params.bWasTriggered);
 						},
 						GET_STATID(STAT_FDelegateGraphTask_StallLogger),
 						Prerequisites,
